LiveLinkLens: Remove the feature where a LiveLink Lens Controller will automatically set the distortion source property of the lens component that it controls.
#rnx

diff --git a/Engine/Plugins/VirtualProduction/LiveLinkLens/Source/LiveLinkLens/Private/LiveLinkLensController.cpp b/Engine/Plugins/VirtualProduction/LiveLinkLens/Source/LiveLinkLens/Private/LiveLinkLensController.cpp
--- a/Engine/Plugins/VirtualProduction/LiveLinkLens/Source/LiveLinkLens/Private/LiveLinkLensController.cpp
+++ b/Engine/Plugins/VirtualProduction/LiveLinkLens/Source/LiveLinkLens/Private/LiveLinkLensController.cpp
@@ -12,21 +12,11 @@
 void ULiveLinkLensController::OnEvaluateRegistered()
 {
 	Super::OnEvaluateRegistered();
-	SetupLensComponent();
 }
 
 void ULiveLinkLensController::SetAttachedComponent(UActorComponent* ActorComponent)
 {
 	Super::SetAttachedComponent(ActorComponent);
-	SetupLensComponent();
-}
-
-void ULiveLinkLensController::SetupLensComponent()
-{
-	if (ULensComponent* const LensComponent = Cast<ULensComponent>(GetAttachedComponent()))
-	{
-		LensComponent->SetDistortionSource(EDistortionSource::LiveLinkLensSubject);
-	}
 }
 
 void ULiveLinkLensController::Tick(float DeltaTime, const FLiveLinkSubjectFrameData& SubjectData)
@@ -83,7 +73,6 @@ TSubclassOf<UActorComponent> ULiveLinkLensController::GetDesiredComponentClass()
 	return ULensComponent::StaticClass();
 }
 
-
 void ULiveLinkLensController::PostLoad()
 {
 	Super::PostLoad();
diff --git a/Engine/Plugins/VirtualProduction/LiveLinkLens/Source/LiveLinkLens/Public/LiveLinkLensController.h b/Engine/Plugins/VirtualProduction/LiveLinkLens/Source/LiveLinkLens/Public/LiveLinkLensController.h
--- a/Engine/Plugins/VirtualProduction/LiveLinkLens/Source/LiveLinkLens/Public/LiveLinkLensController.h
+++ b/Engine/Plugins/VirtualProduction/LiveLinkLens/Source/LiveLinkLens/Public/LiveLinkLensController.h
@@ -30,10 +30,6 @@ public:
 	virtual void PostLoad() override;
 	//~ End UObject interface
 
-private:
-	/** Sets the evaluation mode of the attached LensComponent to use LiveLink and enable applying distortion */
-	void SetupLensComponent();
-
 protected:
 
 #if WITH_EDITORONLY_DATA
