LiveLink: Fix crash when take recording Live Link and Lens Model does not exist.

This may happen if the live link stream specifies a lens model that doesn't exist in Unreal.

#rnx

diff --git a/Engine/Plugins/VirtualProduction/LiveLinkLens/Source/LiveLinkLens/Private/MovieSceneLiveLinkSubSectionLensRole.cpp b/Engine/Plugins/VirtualProduction/LiveLinkLens/Source/LiveLinkLens/Private/MovieSceneLiveLinkSubSectionLensRole.cpp
--- a/Engine/Plugins/VirtualProduction/LiveLinkLens/Source/LiveLinkLens/Private/MovieSceneLiveLinkSubSectionLensRole.cpp
+++ b/Engine/Plugins/VirtualProduction/LiveLinkLens/Source/LiveLinkLens/Private/MovieSceneLiveLinkSubSectionLensRole.cpp
@@ -37,7 +37,11 @@ int32 UMovieSceneLiveLinkSubSectionLensRole::CreateChannelProxy(int32 InChannelI
 	if (SubSystem)
 	{
 		TSubclassOf<ULensModel> LensModelClass = SubSystem->GetRegisteredLensModel(LensData->LensModel);
-		LensModel = LensModelClass->GetDefaultObject<ULensModel>();
+
+		if (LensModelClass.Get())
+		{
+			LensModel = LensModelClass->GetDefaultObject<ULensModel>();
+		}
 	}
 
 	if (LensModel)
@@ -75,7 +79,11 @@ void UMovieSceneLiveLinkSubSectionLensRole::CreatePropertiesChannel()
 	if (SubSystem)
 	{
 		TSubclassOf<ULensModel> LensModelClass = SubSystem->GetRegisteredLensModel(LensData->LensModel);
-		LensModel = LensModelClass->GetDefaultObject<ULensModel>();
+
+		if (LensModelClass.Get())
+		{
+			LensModel = LensModelClass->GetDefaultObject<ULensModel>();
+		}
 	}
 
 	if (LensModel)
